BreakpadSymbolEncoder: If the file doesnt have a newline at EOF handle that as a seperate case

diff --git a/Engine/Source/Programs/BreakpadSymbolEncoder/BreakpadSymbolEncoder.cpp b/Engine/Source/Programs/BreakpadSymbolEncoder/BreakpadSymbolEncoder.cpp
--- a/Engine/Source/Programs/BreakpadSymbolEncoder/BreakpadSymbolEncoder.cpp
+++ b/Engine/Source/Programs/BreakpadSymbolEncoder/BreakpadSymbolEncoder.cpp
@@ -37,13 +37,13 @@ std::string ReadInFile(const std::string& Path)
 
 std::vector<std::string> SplitFilePerLine(std::string&& RawBytes)
 {
-    std::vector<std::string> Out;
+	std::vector<std::string> Out;
 
-    size_t Current = 0;
-    while (Current != std::string::npos && Current < RawBytes.size())
-    {
-        size_t NewLineEnd = RawBytes.find("\n", Current);
-        size_t NewLineSize = 1;
+	size_t Current = 0;
+	while (Current != std::string::npos && Current < RawBytes.size())
+	{
+		size_t NewLineEnd = RawBytes.find("\n", Current);
+		size_t NewLineSize = 1;
 
 		// Check if we are CRLF \r\n
 		if (NewLineEnd > 0 && NewLineEnd != std::string::npos)
@@ -53,14 +53,18 @@ std::vector<std::string> SplitFilePerLine(std::string&& RawBytes)
 				NewLineEnd--;
 				NewLineSize = 2;
 			}
-		}
-
-        Out.push_back(RawBytes.substr(Current, NewLineEnd - Current));
 
-        Current = NewLineEnd + NewLineSize;
-    }
+			Out.push_back(RawBytes.substr(Current, NewLineEnd - Current));
+			Current = NewLineEnd + NewLineSize;
+		}
+		else
+		{
+			Out.push_back(RawBytes.substr(Current));
+			break;
+		}
+	}
 
-    return Out;
+	return Out;
 }
 
 std::vector<std::string> SplitLineIntoNEntries(const std::string& Line, size_t n)
